Fix Linux compile error

diff --git a/Engine/Source/Developer/S3Client/Private/S3Client.cpp b/Engine/Source/Developer/S3Client/Private/S3Client.cpp
--- a/Engine/Source/Developer/S3Client/Private/S3Client.cpp
+++ b/Engine/Source/Developer/S3Client/Private/S3Client.cpp
@@ -746,7 +746,7 @@ FS3ListObjectResponse FS3Client::ListObjects(const FS3ListObjectsRequest& ListRe
 	FS3Response Response = Request.Perform(FS3Request::EMethod::Get, Url.ToString(), FSharedBuffer());
 	if (Response.StatusCode != 200)
 	{
-		return FS3ListObjectResponse{Response.StatusCode, MoveTemp(Response.Body)};
+		return FS3ListObjectResponse{{Response.StatusCode, MoveTemp(Response.Body)}};
 	}
 
 	FString Body(reinterpret_cast<const ANSICHAR*>(Response.Body.GetData()));
@@ -755,14 +755,14 @@ FS3ListObjectResponse FS3Client::ListObjects(const FS3ListObjectsRequest& ListRe
 	if (!XmlFile.IsValid())
 	{
 		//TODO: Better error message
-		return FS3ListObjectResponse{500, MoveTemp(Response.Body)};
+		return FS3ListObjectResponse{{500, MoveTemp(Response.Body)}};
 	}
 
 	const FXmlNode* Root = XmlFile.GetRootNode();
 	if (!Root)
 	{
 		//TODO: Better error message
-		return FS3ListObjectResponse{500, MoveTemp(Response.Body)};
+		return FS3ListObjectResponse{{500, MoveTemp(Response.Body)}};
 	}
 
 	FString BucketName;
@@ -798,7 +798,7 @@ FS3ListObjectResponse FS3Client::ListObjects(const FS3ListObjectsRequest& ListRe
 		}
 	}
 
-	return FS3ListObjectResponse{200, FSharedBuffer(), MoveTemp(BucketName), MoveTemp(Objects)};
+	return FS3ListObjectResponse{{200, FSharedBuffer()}, MoveTemp(BucketName), MoveTemp(Objects)};
 }
 
 } // namespace UE
