S3Client - added dll exports

#rnx

diff --git a/Engine/Source/Developer/S3Client/Internal/S3/S3Client.h b/Engine/Source/Developer/S3Client/Internal/S3/S3Client.h
--- a/Engine/Source/Developer/S3Client/Internal/S3/S3Client.h
+++ b/Engine/Source/Developer/S3Client/Internal/S3/S3Client.h
@@ -11,6 +11,8 @@
 #include "Memory/SharedBuffer.h"
 #include "Misc/StringBuilder.h"
 
+#define UE_API S3CLIENT_API
+
 namespace UE
 {
 
@@ -70,7 +72,7 @@ private:
 /**
  * Container for named credentials.
  */
-class FS3CredentialsProfileStore
+class UE_API FS3CredentialsProfileStore
 {
 public:
 	/** Returns the first credentials loaded into the profile store. */
@@ -166,7 +168,7 @@ struct FS3ClientConfig
 /**
  * A simple HTTP(s) client for down, uploading and listing data objects from Amazon S3.
  */
-class FS3Client
+class UE_API FS3Client
 {
 	class FConnectionPool;
 	class FS3Request;
@@ -205,3 +207,5 @@ private:
 } // namesapce UE
 
 #endif // (IS_PROGRAM || WITH_EDITOR)
+
+#undef UE_API
