Only include S3Client module on desktop platforms to fix curl include error

#rnx

diff --git a/Engine/Source/Developer/S3Client/Private/S3Client.cpp b/Engine/Source/Developer/S3Client/Private/S3Client.cpp
--- a/Engine/Source/Developer/S3Client/Private/S3Client.cpp
+++ b/Engine/Source/Developer/S3Client/Private/S3Client.cpp
@@ -8,12 +8,6 @@
 
 #if (IS_PROGRAM || WITH_EDITOR)
 
-#if PLATFORM_MICROSOFT
-	#include "Microsoft/WindowsHWrapper.h"
-	#include "Microsoft/AllowMicrosoftPlatformTypes.h"
-	#include "Microsoft/HideMicrosoftPlatformTypes.h"
-#endif
-
 #include "Containers/StringConv.h"
 #include "Misc/CString.h"
 #include "Misc/DateTime.h"
@@ -21,8 +15,17 @@
 #include "String/BytesToHex.h"
 #include "XmlParser.h"
 
-//TODO: Switch from Curl to CoreHttp
+#if PLATFORM_MICROSOFT
+#include "Microsoft/AllowMicrosoftPlatformTypes.h"
+#endif
+#ifdef PLATFORM_CURL_INCLUDE
+#include PLATFORM_CURL_INCLUDE
+#else
 #include "curl/curl.h"
+#endif
+#if PLATFORM_MICROSOFT
+#include "Microsoft/HideMicrosoftPlatformTypes.h"
+#endif
 
 #if WITH_SSL
 #include "Ssl.h"
diff --git a/Engine/Source/Runtime/Experimental/IoStoreOnDemand/IoStoreOnDemand.Build.cs b/Engine/Source/Runtime/Experimental/IoStoreOnDemand/IoStoreOnDemand.Build.cs
--- a/Engine/Source/Runtime/Experimental/IoStoreOnDemand/IoStoreOnDemand.Build.cs
+++ b/Engine/Source/Runtime/Experimental/IoStoreOnDemand/IoStoreOnDemand.Build.cs
@@ -10,7 +10,8 @@ public class IoStoreOnDemand : ModuleRules
         PrivateDependencyModuleNames.AddRange(new string[] { "HTTP", "Json" });
 		bAllowConfidentialPlatformDefines = true;
 
-		if (Target.Type == TargetType.Editor || Target.Type == TargetType.Program)
+		if (Target.Platform.IsInGroup(UnrealPlatformGroup.Desktop) &&
+			(Target.Type == TargetType.Editor || Target.Type == TargetType.Program))
 		{
 			PrivateDependencyModuleNames.AddRange(new string[] { "S3Client", "RSA" });
 		}
