Add human readable error messages when failing to upload IAS chunks.

#rnx

diff --git a/Engine/Source/Developer/S3Client/Internal/S3/S3Client.h b/Engine/Source/Developer/S3Client/Internal/S3/S3Client.h
--- a/Engine/Source/Developer/S3Client/Internal/S3/S3Client.h
+++ b/Engine/Source/Developer/S3Client/Internal/S3/S3Client.h
@@ -105,6 +105,8 @@ struct FS3Response
 	/** Returns the body as text. */
 	FString ToString() const { return FString(reinterpret_cast<const ANSICHAR*>(Body.GetData())); }
 
+	UE_API void GetErrorMsg(FStringBuilderBase& OutErrorMsg) const;
+
 	/** HTTP status code. */
 	uint32 StatusCode = 0;
 	/** HTTP response body. */
diff --git a/Engine/Source/Developer/S3Client/Private/S3Client.cpp b/Engine/Source/Developer/S3Client/Private/S3Client.cpp
--- a/Engine/Source/Developer/S3Client/Private/S3Client.cpp
+++ b/Engine/Source/Developer/S3Client/Private/S3Client.cpp
@@ -326,6 +326,50 @@ FS3CredentialsProfileStore FS3CredentialsProfileStore::FromFile(const FString& F
 	return ProfileStore;
 }
 
+////////////////////////////////////////////////////////////////////////////////
+void FS3Response::GetErrorMsg(FStringBuilderBase& OutErrorMsg) const
+{
+	OutErrorMsg.Reset();
+
+	if (IsOk())
+	{
+		OutErrorMsg << TEXT("Successs");
+		return;
+	}
+
+	const FString BodyResponseString = ToString();
+
+	FXmlFile XmlFile;
+	if (!XmlFile.LoadFile(BodyResponseString, EConstructMethod::ConstructFromBuffer))
+	{
+		OutErrorMsg << TEXT("Unknown");
+		return;
+	}
+
+	const FXmlNode* RootNode = XmlFile.GetRootNode();
+	if (!RootNode)
+	{
+		OutErrorMsg << TEXT("Unknown");
+		return;
+	}
+
+	const FXmlNode* CodeNode = RootNode->FindChildNode(TEXT("Code"));
+	if (!CodeNode)
+	{
+		OutErrorMsg << TEXT("Unknown");
+		return;
+	}
+
+	const FXmlNode* MessageNode = RootNode->FindChildNode(TEXT("Message"));
+	if (!MessageNode)
+	{
+		OutErrorMsg << TEXT("Unknown");
+		return;
+	}
+
+	OutErrorMsg << CodeNode->GetContent() << TEXT(": ") << MessageNode->GetContent();
+}
+
 ////////////////////////////////////////////////////////////////////////////////
 class FS3Client::FConnectionPool
 {
diff --git a/Engine/Source/Runtime/Experimental/IoStoreOnDemand/Private/IoStoreOnDemand.cpp b/Engine/Source/Runtime/Experimental/IoStoreOnDemand/Private/IoStoreOnDemand.cpp
--- a/Engine/Source/Runtime/Experimental/IoStoreOnDemand/Private/IoStoreOnDemand.cpp
+++ b/Engine/Source/Runtime/Experimental/IoStoreOnDemand/Private/IoStoreOnDemand.cpp
@@ -706,7 +706,10 @@ void FS3UploadQueue::ThreadEntry()
 		}
 		else
 		{
-			UE_LOG(LogIas, Warning, TEXT("Failed to upload chunk '%s/%s/%s', StatusCode: %u"), *Client.GetConfig().ServiceUrl, *Bucket, *Entry.Key, Response.StatusCode);
+			TStringBuilder<256> ErrorResponse;
+			Response.GetErrorMsg(ErrorResponse);
+
+			UE_LOG(LogIas, Warning, TEXT("Failed to upload chunk '%s/%s/%s', StatusCode: %u, Error: %s"), *Client.GetConfig().ServiceUrl, *Bucket, *Entry.Key, Response.StatusCode, ErrorResponse.ToString());
 			ErrorCount++;
 			break;
 		}
