S3 Client - delete object

#rnx

diff --git a/Engine/Source/Developer/S3Client/Internal/S3/S3Client.h b/Engine/Source/Developer/S3Client/Internal/S3/S3Client.h
--- a/Engine/Source/Developer/S3Client/Internal/S3/S3Client.h
+++ b/Engine/Source/Developer/S3Client/Internal/S3/S3Client.h
@@ -160,6 +160,12 @@ struct FS3ListObjectResponse
 	TArray<FS3Object> Objects;
 };
 
+/** Request parameters for deleting objects. */
+using FS3DeleteObjectRequest = FS3GetObjectRequest;
+
+/** Response parameters when deleting objects. */
+using FS3DeleteObjectResponse = FS3Response;
+
 /** S3 client configuration with region and service URL. */
 struct FS3ClientConfig
 {
@@ -198,6 +204,8 @@ public:
 	FS3PutObjectResponse PutObject(const FS3PutObjectRequest& Request);
 	/** Retries uploading an object until succeeded or max attempts has been reached. */
 	FS3PutObjectResponse TryPutObject(const FS3PutObjectRequest& Request, int32 MaxAttempts = 3, float Delay = 0.5f);
+	/** Delete an object associated with the specified request parameters. */
+	FS3DeleteObjectResponse DeleteObject(const FS3DeleteObjectRequest& Request);
 
 private:
 	void Setup(FS3Request& Request);
diff --git a/Engine/Source/Developer/S3Client/Private/S3Client.cpp b/Engine/Source/Developer/S3Client/Private/S3Client.cpp
--- a/Engine/Source/Developer/S3Client/Private/S3Client.cpp
+++ b/Engine/Source/Developer/S3Client/Private/S3Client.cpp
@@ -446,7 +446,8 @@ public:
 	enum class EMethod
 	{
 		Get,
-		Put
+		Put,
+		Delete,
 	};
 
 	FS3Request(FS3Client& S3Client);
@@ -461,6 +462,7 @@ private:
 	static size_t WriteBodyCallback(void* Ptr, size_t SizeInBlocks, size_t BlockSizeInBytes, void* UserData);
 	static int SslCertVerify(int PreverifyOk, X509_STORE_CTX* Context);
 	static CURLcode SslContextCallback(CURL* curl, void* sslctx, void* parm);
+	const ANSICHAR* LexToString(EMethod Method);
 
 	FS3Client& Client;
 	FCurlHandle Curl;
@@ -489,7 +491,7 @@ FS3Response FS3Client::FS3Request::Perform(EMethod Method, const ANSICHAR* Url,
 
 	RequestBody = Body;
 	const uint64 ContentLength = Body.GetSize();
-	check(Method == EMethod::Get || ContentLength > 0);
+	check((Method == EMethod::Get || Method == EMethod::Delete) || ContentLength > 0);
 
 	// Find the host from the URL
 	const ANSICHAR* ProtocolEnd = FCStringAnsi::Strchr(Url, ':');
@@ -536,7 +538,7 @@ FS3Response FS3Client::FS3Request::Perform(EMethod Method, const ANSICHAR* Url,
 		CurlHeaders = curl_slist_append(CurlHeaders, *WriteToAnsiString<512>("x-amz-security-token: ", *Client.Credentials.GetSessionToken()));
 	}
 
-	const ANSICHAR* MethodString = Method == EMethod::Get ? "GET" : "PUT";
+	const ANSICHAR* MethodString = LexToString(Method);
 	CurlHeaders = curl_slist_append(CurlHeaders, GetAuthorizationHeader(Curl, Client, MethodString, UrlHostEnd, QueryString, CurlHeaders, *TimeString, *PayloadSha256, AuthHeader));
 
 	// Append the unsigned headers
@@ -553,7 +555,7 @@ FS3Response FS3Client::FS3Request::Perform(EMethod Method, const ANSICHAR* Url,
 	{
 		curl_easy_setopt(Curl, CURLOPT_HTTPGET, 1L);
 	}
-	else
+	else if (Method == EMethod::Put)
 	{
 		curl_easy_setopt(Curl, CURLOPT_PUT, 1L);
 		curl_easy_setopt(Curl, CURLOPT_UPLOAD, 1L);
@@ -561,6 +563,11 @@ FS3Response FS3Client::FS3Request::Perform(EMethod Method, const ANSICHAR* Url,
 		curl_easy_setopt(Curl, CURLOPT_READDATA, this);
 		curl_easy_setopt(Curl, CURLOPT_READFUNCTION, &ReadBodyCallback);
 	}
+	else
+	{
+		check(Method == EMethod::Delete);
+		curl_easy_setopt(Curl, CURLOPT_CUSTOMREQUEST, "DELETE");
+	}
 
 	curl_easy_setopt(Curl, CURLOPT_FOLLOWLOCATION, 1L);
 	curl_easy_setopt(Curl, CURLOPT_NOSIGNAL, 1L);
@@ -725,6 +732,22 @@ CURLcode FS3Client::FS3Request::SslContextCallback(CURL* curl, void* sslctx, voi
 	return CURLE_OK;
 }
 
+const ANSICHAR* FS3Client::FS3Request::LexToString(EMethod Method)
+{
+	switch(Method)
+	{
+	case EMethod::Get:
+		return "GET";
+	case EMethod::Put:
+		return "PUT";
+	case EMethod::Delete:
+		return "DELETE";
+	default:
+		check(false);
+		return nullptr;
+	}
+}
+
 ////////////////////////////////////////////////////////////////////////////////
 FS3Client::FS3Client(const FS3ClientConfig& ClientConfig, const FS3ClientCredentials& ClientCredentials)
 	: Config(ClientConfig)
@@ -868,6 +891,15 @@ FS3ListObjectResponse FS3Client::ListObjects(const FS3ListObjectsRequest& ListRe
 	return FS3ListObjectResponse{{200, FSharedBuffer()}, MoveTemp(BucketName), MoveTemp(Objects)};
 }
 
+FS3DeleteObjectResponse FS3Client::DeleteObject(const FS3DeleteObjectRequest& DeleteRequest)
+{
+	TAnsiStringBuilder<256> Url;
+	Url << StringCast<ANSICHAR>(*Config.ServiceUrl) << "/" << StringCast<ANSICHAR>(*DeleteRequest.BucketName) << "/" << StringCast<ANSICHAR>(*DeleteRequest.Key);
+
+	FS3Request Request(*this);
+	return Request.Perform(FS3Request::EMethod::Delete, Url.ToString(), FSharedBuffer());
+}
+
 } // namespace UE
 
 #endif // (IS_PROGRAM || WITH_EDITOR)
