Improve error messages when failing to upload IAS chunks.

#rnx

### Error Handling
- Previously we only recorded the http response value inFS3Response, which meant that issues with the API itself (such as being unable to connect to the target) would go unreported.
- We now record the curl error value as FS3Response::ApiStatusCode
- Replaced FS3Response::GetErrorMsg with two different methods, one to return a short error message based off the api and http status values and another to return longer error messages with additional error messages parsed from the response body where possible.
-- We do also log the api (curl) error when they are first detected under the "LogS3Client" category, but these messages are too verbose and heavy weight to pass around with the FS3Response class.

### Misc
- Add a new logging category "LogS3Client" for the S3Client module for now it is declared in the only cpp of the module.
- Made the members of FS3Response private and added the appropriate accessors.
-- This should encourage people to use the accessors provided to get error messages rather than just relying on the HttpStatusCode or ApiStatusCode and maybe missing out on important info.

diff --git a/Engine/Source/Developer/S3Client/Internal/S3/S3Client.h b/Engine/Source/Developer/S3Client/Internal/S3/S3Client.h
--- a/Engine/Source/Developer/S3Client/Internal/S3/S3Client.h
+++ b/Engine/Source/Developer/S3Client/Internal/S3/S3Client.h
@@ -100,15 +100,54 @@ struct FS3Object
 /** Basic response parameters. */
 struct FS3Response
 {
-	/** Returns whether the request is considered successful. */
-	bool IsOk() const { return StatusCode > 199 && StatusCode < 299; }
+public:
+
+	FS3Response();
+	
+	FS3Response(const FS3Response& Other);
+	FS3Response(FS3Response&& Other);
+
+	FS3Response(uint32 InHttpStatusCode, uint32 InApiStatusCode);
+	FS3Response(uint32 InHttpStatusCode, uint32 InApiStatusCode, FSharedBuffer&& InBody);
+	FS3Response(uint32 InHttpStatusCode, FS3Response&& Other);
+
+public:
+
+	UE_API ~FS3Response() = default;
+
+	UE_API FS3Response& operator=(const FS3Response& Other);
+	UE_API FS3Response& operator=(FS3Response&& Other);
+
+	/** Returns whether the request is considered successful or not. */
+	UE_API bool IsOk() const;
+
 	/** Returns the body as text. */
-	FString ToString() const { return FString(reinterpret_cast<const ANSICHAR*>(Body.GetData())); }
+	UE_API FString ToString() const;
+
+	/** Returns the body as a raw buffer */
+	UE_API FSharedBuffer GetBody() const;
+
+	/** 
+	 * Returns an error message comprised of the api/http status code and any error
+	 * that might have been returned as part of the response body.
+	 * If there response had no errors it the message will be "Success"
+	 */
+	UE_API void GetErrorResponse(FStringBuilderBase& OutErrorMsg) const;
+	
+	/** 
+	 * Returns a short error message comprised of the api/http status code
+	 * If there response had no errors it the message will be "Success"
+	 */
+	UE_API FString GetErrorStatus() const;
+
+private:
+
+	/** Status code returned by the HTTP request. */
+	uint32 HttpStatusCode = 0;
 
-	UE_API void GetErrorMsg(FStringBuilderBase& OutErrorMsg) const;
+	/** Status code returned by the API */
+	uint32 ApiStatusCode = 0;
 
-	/** HTTP status code. */
-	uint32 StatusCode = 0;
 	/** HTTP response body. */
 	FSharedBuffer Body;
 };
diff --git a/Engine/Source/Developer/S3Client/Private/S3Client.cpp b/Engine/Source/Developer/S3Client/Private/S3Client.cpp
--- a/Engine/Source/Developer/S3Client/Private/S3Client.cpp
+++ b/Engine/Source/Developer/S3Client/Private/S3Client.cpp
@@ -34,6 +34,8 @@
 
 IMPLEMENT_MODULE(FDefaultModuleImpl, S3Client);
 
+DEFINE_LOG_CATEGORY_STATIC(LogS3Client, Log, All);
+
 namespace UE
 {
 
@@ -325,7 +327,69 @@ FS3CredentialsProfileStore FS3CredentialsProfileStore::FromFile(const FString& F
 }
 
 ////////////////////////////////////////////////////////////////////////////////
-void FS3Response::GetErrorMsg(FStringBuilderBase& OutErrorMsg) const
+
+FS3Response::FS3Response() = default;
+
+FS3Response::FS3Response(const FS3Response& Other) = default;
+
+FS3Response::FS3Response(FS3Response&& Other)
+	: HttpStatusCode(Other.HttpStatusCode)
+	, ApiStatusCode(Other.ApiStatusCode)
+	, Body(MoveTemp(Other.Body))
+{
+
+}
+
+FS3Response::FS3Response(uint32 InHttpStatusCode, uint32 InApiStatusCode)
+	: HttpStatusCode(InHttpStatusCode)
+	, ApiStatusCode(InApiStatusCode)
+{
+
+}
+
+FS3Response::FS3Response(uint32 InHttpStatusCode, uint32 InApiStatusCode, FSharedBuffer&& InBody)
+	: HttpStatusCode(InHttpStatusCode)
+	, ApiStatusCode(InApiStatusCode)
+	, Body(MoveTemp(InBody))
+{
+
+}
+
+FS3Response::FS3Response(uint32 InHttpStatusCode, FS3Response&& Other)
+	: HttpStatusCode(InHttpStatusCode)
+	, ApiStatusCode(Other.ApiStatusCode)
+	, Body(MoveTemp(Other.Body))
+{
+
+}
+
+FS3Response& FS3Response::operator=(const FS3Response& Other) = default;
+
+FS3Response& FS3Response::operator=(FS3Response&& Other)
+{
+	HttpStatusCode = Other.HttpStatusCode;
+	ApiStatusCode = Other.ApiStatusCode;
+	Body = MoveTemp(Other.Body);
+
+	return *this;
+}
+
+bool FS3Response::IsOk() const
+{
+	return ApiStatusCode == CURLE_OK && HttpStatusCode > 199 && HttpStatusCode < 299;
+}
+
+FSharedBuffer FS3Response::GetBody() const
+{
+	return Body;
+}
+
+FString FS3Response::ToString() const
+{
+	return FString(reinterpret_cast<const ANSICHAR*>(Body.GetData()));
+}
+
+void FS3Response::GetErrorResponse(FStringBuilderBase& OutErrorMsg) const
 {
 	OutErrorMsg.Reset();
 
@@ -335,37 +399,66 @@ void FS3Response::GetErrorMsg(FStringBuilderBase& OutErrorMsg) const
 		return;
 	}
 
-	const FString BodyResponseString = ToString();
-
-	FXmlFile XmlFile;
-	if (!XmlFile.LoadFile(BodyResponseString, EConstructMethod::ConstructFromBuffer))
+	if(ApiStatusCode != CURLE_OK)
 	{
-		OutErrorMsg << TEXT("Unknown");
-		return;
+		OutErrorMsg << curl_easy_strerror(static_cast<CURLcode>(ApiStatusCode));
 	}
-
-	const FXmlNode* RootNode = XmlFile.GetRootNode();
-	if (!RootNode)
+	else
 	{
-		OutErrorMsg << TEXT("Unknown");
-		return;
+		OutErrorMsg << TEXT("StatusCode: ") << HttpStatusCode << TEXT("Error: ");
+
+		const FString BodyResponseString = ToString();
+
+		FXmlFile XmlFile;
+		if (!XmlFile.LoadFile(BodyResponseString, EConstructMethod::ConstructFromBuffer))
+		{
+			OutErrorMsg << TEXT("Unknown");
+			return;
+		}
+
+		const FXmlNode* RootNode = XmlFile.GetRootNode();
+		if (!RootNode)
+		{
+			OutErrorMsg << TEXT("Unknown");
+			return;
+		}
+
+		const FXmlNode* CodeNode = RootNode->FindChildNode(TEXT("Code"));
+		if (!CodeNode)
+		{
+			OutErrorMsg << TEXT("Unknown");
+			return;
+		}
+
+		const FXmlNode* MessageNode = RootNode->FindChildNode(TEXT("Message"));
+		if (!MessageNode)
+		{
+			OutErrorMsg << TEXT("Unknown");
+			return;
+		}
+
+		OutErrorMsg << CodeNode->GetContent() << TEXT(": ") << MessageNode->GetContent();
 	}
+}
+
+FString FS3Response:: GetErrorStatus() const
+{
+	TStringBuilder<256> ErrorStatus;
 
-	const FXmlNode* CodeNode = RootNode->FindChildNode(TEXT("Code"));
-	if (!CodeNode)
+	if (IsOk())
 	{
-		OutErrorMsg << TEXT("Unknown");
-		return;
+		ErrorStatus << TEXT("Successs");
 	}
-
-	const FXmlNode* MessageNode = RootNode->FindChildNode(TEXT("Message"));
-	if (!MessageNode)
+	else if (ApiStatusCode != CURLE_OK)
 	{
-		OutErrorMsg << TEXT("Unknown");
-		return;
+		ErrorStatus << curl_easy_strerror(static_cast<CURLcode>(ApiStatusCode));
+	}
+	else
+	{
+		ErrorStatus << TEXT("StatusCode: ") << HttpStatusCode << TEXT("Error: ");
 	}
 
-	OutErrorMsg << CodeNode->GetContent() << TEXT(": ") << MessageNode->GetContent();
+	return FString(ErrorStatus);
 }
 
 ////////////////////////////////////////////////////////////////////////////////
@@ -583,6 +676,10 @@ FS3Response FS3Client::FS3Request::Perform(EMethod Method, const ANSICHAR* Url,
 	curl_easy_setopt(Curl, CURLOPT_WRITEDATA, this);
 	curl_easy_setopt(Curl, CURLOPT_WRITEFUNCTION, &WriteBodyCallback);
 
+	// Errors and logging
+	ANSICHAR ErrorBuffer[CURL_ERROR_SIZE] = {0};
+	curl_easy_setopt(Curl, CURLOPT_ERRORBUFFER, ErrorBuffer);
+
 	// SSL options
 	curl_easy_setopt(Curl, CURLOPT_USE_SSL, CURLUSESSL_ALL);
 	curl_easy_setopt(Curl, CURLOPT_SSL_VERIFYPEER, 1);
@@ -607,17 +704,31 @@ FS3Response FS3Client::FS3Request::Perform(EMethod Method, const ANSICHAR* Url,
 	{
 		CurlResult = curl_easy_getinfo(Curl, CURLINFO_RESPONSE_CODE, &ResponseCode);
 	}
+	
+	if (CurlResult != CURLE_OK)
+	{
+		UE_LOG(LogS3Client, Error, TEXT("%hs"), ErrorBuffer);
+	}
+
+	const char* ErrorMsg = curl_easy_strerror(CurlResult);
 
 	if (const uint64 Size = ResponseBody.TotalSize(); Size > 0)
 	{
 		return FS3Response
 		{
 			static_cast<uint32>(ResponseCode),
+			static_cast<uint32>(CurlResult),
 			FSharedBuffer::TakeOwnership(ResponseBody.ReleaseOwnership(), Size, FMemory::Free)
 		};
 	}
-
-	return FS3Response{static_cast<uint32>(ResponseCode)};
+	else
+	{
+		return FS3Response
+		{
+			static_cast<uint32>(ResponseCode),
+			static_cast<uint32>(CurlResult)
+		};
+	}
 }
 
 int FS3Client::FS3Request::StatusCallback(void* Ptr, curl_off_t TotalDownloadSize, curl_off_t CurrentDownloadSize, curl_off_t TotalUploadSize, curl_off_t CurrentUploadSize)
@@ -834,25 +945,25 @@ FS3ListObjectResponse FS3Client::ListObjects(const FS3ListObjectsRequest& ListRe
 
 	FS3Request Request(*this);
 	FS3Response Response = Request.Perform(FS3Request::EMethod::Get, Url.ToString(), FSharedBuffer());
-	if (Response.StatusCode != 200)
+	if (!Response.IsOk())
 	{
-		return FS3ListObjectResponse{{Response.StatusCode, MoveTemp(Response.Body)}};
+		return FS3ListObjectResponse(MoveTemp(Response));
 	}
 
-	FString Body(reinterpret_cast<const ANSICHAR*>(Response.Body.GetData()));
+	FString Body(Response.ToString());
 
 	FXmlFile XmlFile(Body, EConstructMethod::ConstructFromBuffer);
 	if (!XmlFile.IsValid())
 	{
 		//TODO: Better error message
-		return FS3ListObjectResponse{{500, MoveTemp(Response.Body)}};
+		return FS3ListObjectResponse{{500, MoveTemp(Response)}};
 	}
 
 	const FXmlNode* Root = XmlFile.GetRootNode();
 	if (!Root)
 	{
 		//TODO: Better error message
-		return FS3ListObjectResponse{{500, MoveTemp(Response.Body)}};
+		return FS3ListObjectResponse{{500, MoveTemp(Response)}};
 	}
 
 	FString BucketName;
@@ -888,7 +999,7 @@ FS3ListObjectResponse FS3Client::ListObjects(const FS3ListObjectsRequest& ListRe
 		}
 	}
 
-	return FS3ListObjectResponse{{200, FSharedBuffer()}, MoveTemp(BucketName), MoveTemp(Objects)};
+	return FS3ListObjectResponse{{200, 0, FSharedBuffer()}, MoveTemp(BucketName), MoveTemp(Objects)};
 }
 
 FS3DeleteObjectResponse FS3Client::DeleteObject(const FS3DeleteObjectRequest& DeleteRequest)
diff --git a/Engine/Source/Runtime/Experimental/IoStoreOnDemand/Private/IoStoreOnDemand.cpp b/Engine/Source/Runtime/Experimental/IoStoreOnDemand/Private/IoStoreOnDemand.cpp
--- a/Engine/Source/Runtime/Experimental/IoStoreOnDemand/Private/IoStoreOnDemand.cpp
+++ b/Engine/Source/Runtime/Experimental/IoStoreOnDemand/Private/IoStoreOnDemand.cpp
@@ -897,9 +897,9 @@ void FS3UploadQueue::ThreadEntry()
 		else
 		{
 			TStringBuilder<256> ErrorResponse;
-			Response.GetErrorMsg(ErrorResponse);
+			Response.GetErrorResponse(ErrorResponse);
 
-			UE_LOG(LogIas, Warning, TEXT("Failed to upload chunk '%s/%s/%s', StatusCode: %u, Error: %s"), *Client.GetConfig().ServiceUrl, *Bucket, *Entry.Key, Response.StatusCode, ErrorResponse.ToString());
+			UE_LOG(LogIas, Warning, TEXT("Failed to upload chunk '%s/%s/%s' (%s)"), *Client.GetConfig().ServiceUrl, *Bucket, *Entry.Key, ErrorResponse.ToString());
 			ErrorCount++;
 			break;
 		}
@@ -1268,13 +1268,13 @@ TIoStatusOr<FIoStoreUploadResult> UploadContainerFiles(
 			}
 
 			FOnDemandToc Toc;
-			FMemoryReaderView Ar(TocResponse.Body.GetView());
+			FMemoryReaderView Ar(TocResponse.GetBody().GetView());
 			Ar << Toc;
 
 			if (Ar.IsError()) 
 			{
 				Toc = FOnDemandToc{};
-				if (LoadFromCompactBinary(FCbFieldView(TocResponse.Body.GetData()), Toc) == false)
+				if (LoadFromCompactBinary(FCbFieldView(TocResponse.GetBody().GetData()), Toc) == false)
 				{
 					UE_LOG(LogIas, Warning, TEXT("Failed to load TOC '%s/%s/%s'"), *Client.GetConfig().ServiceUrl, *UploadParams.Bucket, *TocInfo.Key);
 					continue;
@@ -1456,7 +1456,7 @@ TIoStatusOr<FIoStoreUploadResult> UploadContainerFiles(
 
 			if (!Response.IsOk())
 			{
-				return FIoStatus(EIoErrorCode::WriteError, FString::Printf(TEXT("Failed to upload '%s', StatusCode: %u"), *UTocFilePath, Response.StatusCode));
+				return FIoStatus(EIoErrorCode::WriteError, FString::Printf(TEXT("Failed to upload '%s' (%s)"), *UTocFilePath, *Response.GetErrorStatus()));
 			}
 			
 			UTocPaths.Add(ContainerEntry.UTocHash, UTocFilePath);
@@ -1519,7 +1519,7 @@ TIoStatusOr<FIoStoreUploadResult> UploadContainerFiles(
 		}
 		else
 		{
-			UE_LOG(LogIas, Warning, TEXT("Failed to upload TOC '%s/%s/%s', StatusCode: %u"), *Client.GetConfig().ServiceUrl, *UploadParams.Bucket, Key.ToString(), Response.StatusCode);
+			UE_LOG(LogIas, Warning, TEXT("Failed to upload TOC '%s/%s/%s' (%s)"), *Client.GetConfig().ServiceUrl, *UploadParams.Bucket, Key.ToString(), *Response.GetErrorStatus());
 			return FIoStatus(EIoErrorCode::WriteError, TEXT("Failed to upload TOC"));
 		}
 
@@ -1716,7 +1716,7 @@ FIoStatus DownloadContainerFiles(const FIoStoreDownloadParams& DownloadParams, c
 	}
 
 	FOnDemandToc OnDemandToc;
-	if (LoadFromCompactBinary(FCbFieldView(TocResponse.Body.GetData()), OnDemandToc) == false)
+	if (LoadFromCompactBinary(FCbFieldView(TocResponse.GetBody().GetData()), OnDemandToc) == false)
 	{
 		return FIoStatus(EIoErrorCode::ReadError, TEXT("Failed to load on demand TOC"));
 	}
@@ -1755,12 +1755,12 @@ FIoStatus DownloadContainerFiles(const FIoStoreDownloadParams& DownloadParams, c
 		const FString UTocPath = FString::Printf(TEXT("%s/%s.utoc"), *DownloadParams.Directory, *ContainerEntry.ContainerName);
 		const FString UCasPath = FPaths::ChangeExtension(UTocPath, TEXT(".ucas")); 
 		FContainerStats& ContainerStats = ContainerSummary.FindOrAdd(ContainerEntry.ContainerName);
-		ContainerStats.TocRawSize = Response.Body.GetSize();
+		ContainerStats.TocRawSize = Response.GetBody().GetSize();
 
 		if (TUniquePtr<FArchive> TocFile(IFileManager::Get().CreateFileWriter(*UTocPath)); TocFile.IsValid())
 		{
 			UE_LOG(LogIas, Display, TEXT("Writing '%s'"), *UTocPath);
-			TocFile->Serialize((void*)Response.Body.GetData(), Response.Body.GetSize());
+			TocFile->Serialize((void*)Response.GetBody().GetData(), Response.GetBody().GetSize());
 		}
 		else
 		{
@@ -1850,10 +1850,10 @@ FIoStatus DownloadContainerFiles(const FIoStoreDownloadParams& DownloadParams, c
 			}
 
 			UE_LOG(LogIas, Display, TEXT("Serializing chunk %d/%d '%s' -> '%s' (%llu B)"),
-				Idx + 1, TocEntryCount, *HashString, *LexToString(ChunkId), ChunkResponse.Body.GetSize());
+				Idx + 1, TocEntryCount, *HashString, *LexToString(ChunkId), ChunkResponse.GetBody().GetSize());
 
 			check(CasFile->Tell() == FirstBlock.GetOffset());
-			CasFile->Serialize((void*)ChunkResponse.Body.GetData(), ChunkResponse.Body.GetSize());
+			CasFile->Serialize((void*)ChunkResponse.GetBody().GetData(), ChunkResponse.GetBody().GetSize());
 		}
 
 		ContainerStats.CompressedSize = CasFile->Tell();
@@ -2119,7 +2119,7 @@ FIoStatus ListTocs(const FIoStoreListTocsParams& Params)
 
 		if (Response.Objects.IsEmpty())
 		{
-			UE_LOG(LogIas, Display, TEXT("Not TOC's found at '%s/%s/%s' (StatusCode: %u)"), *Client.GetConfig().ServiceUrl, *Params.Bucket, Path.ToString(), Response.StatusCode);
+			UE_LOG(LogIas, Display, TEXT("Not TOC's found at '%s/%s/%s' (%s)"), *Client.GetConfig().ServiceUrl, *Params.Bucket, Path.ToString(), *Response.GetErrorStatus());
 			return FIoStatus(EIoErrorCode::NotFound);
 		}
 
@@ -2150,7 +2150,7 @@ FIoStatus ListTocs(const FIoStoreListTocsParams& Params)
 			}
 
 			FOnDemandToc Toc;
-			FMemoryReaderView Ar(TocResponse.Body.GetView());
+			FMemoryReaderView Ar(TocResponse.GetBody().GetView());
 			Ar << Toc;
 
 			if (Toc.Header.Magic != FOnDemandTocHeader::ExpectedMagic)
@@ -2196,7 +2196,7 @@ FIoStatus ListTocs(const FIoStoreListTocsParams& Params)
 		if (TocResponse.IsOk())
 		{
 			FOnDemandToc Toc;
-			FMemoryReaderView Ar(TocResponse.Body.GetView());
+			FMemoryReaderView Ar(TocResponse.GetBody().GetView());
 			Ar << Toc;
 
 			if (Toc.Header.Magic == FOnDemandTocHeader::ExpectedMagic)
@@ -2207,7 +2207,7 @@ FIoStatus ListTocs(const FIoStoreListTocsParams& Params)
 					.Toc = MoveTemp(Toc),
 					.DateTime = DateTime,
 					.Key = Params.TocKey,
-					.Size = TocResponse.Body.GetSize()
+					.Size = TocResponse.GetBody().GetSize()
 				});
 			}
 			else
