[Backout] - CL32362006
Original CL Desc
-----------------------------------------------------------------
Improve error messages when failing to upload IAS chunks.

#rnx

### Error Handling
- Previously we only recorded the http response value inFS3Response, which meant that issues with the API itself (such as being unable to connect to the target) would go unreported.
- We now record the curl error value as FS3Response::ApiStatusCode
- Replaced FS3Response::GetErrorMsg with two different methods, one to return a short error message based off the api and http status values and another to return longer error messages with additional error messages parsed from the response body where possible.
-- We do also log the api (curl) error when they are first detected under the "LogS3Client" category, but these messages are too verbose and heavy weight to pass around with the FS3Response class.

### Misc
- Add a new logging category "LogS3Client" for the S3Client module for now it is declared in the only cpp of the module.
- Made the members of FS3Response private and added the appropriate accessors.
-- This should encourage people to use the accessors provided to get error messages rather than just relying on the HttpStatusCode or ApiStatusCode and maybe missing out on important info.

diff --git a/Engine/Source/Developer/S3Client/Internal/S3/S3Client.h b/Engine/Source/Developer/S3Client/Internal/S3/S3Client.h
--- a/Engine/Source/Developer/S3Client/Internal/S3/S3Client.h
+++ b/Engine/Source/Developer/S3Client/Internal/S3/S3Client.h
@@ -100,54 +100,15 @@ struct FS3Object
 /** Basic response parameters. */
 struct FS3Response
 {
-public:
-
-	FS3Response();
-	
-	FS3Response(const FS3Response& Other);
-	FS3Response(FS3Response&& Other);
-
-	FS3Response(uint32 InHttpStatusCode, uint32 InApiStatusCode);
-	FS3Response(uint32 InHttpStatusCode, uint32 InApiStatusCode, FSharedBuffer&& InBody);
-	FS3Response(uint32 InHttpStatusCode, FS3Response&& Other);
-
-public:
-
-	UE_API ~FS3Response() = default;
-
-	UE_API FS3Response& operator=(const FS3Response& Other);
-	UE_API FS3Response& operator=(FS3Response&& Other);
-
-	/** Returns whether the request is considered successful or not. */
-	UE_API bool IsOk() const;
-
+	/** Returns whether the request is considered successful. */
+	bool IsOk() const { return StatusCode > 199 && StatusCode < 299; }
 	/** Returns the body as text. */
-	UE_API FString ToString() const;
-
-	/** Returns the body as a raw buffer */
-	UE_API FSharedBuffer GetBody() const;
-
-	/** 
-	 * Returns an error message comprised of the api/http status code and any error
-	 * that might have been returned as part of the response body.
-	 * If there response had no errors it the message will be "Success"
-	 */
-	UE_API void GetErrorResponse(FStringBuilderBase& OutErrorMsg) const;
-	
-	/** 
-	 * Returns a short error message comprised of the api/http status code
-	 * If there response had no errors it the message will be "Success"
-	 */
-	UE_API FString GetErrorStatus() const;
-
-private:
-
-	/** Status code returned by the HTTP request. */
-	uint32 HttpStatusCode = 0;
+	FString ToString() const { return FString(reinterpret_cast<const ANSICHAR*>(Body.GetData())); }
 
-	/** Status code returned by the API */
-	uint32 ApiStatusCode = 0;
+	UE_API void GetErrorMsg(FStringBuilderBase& OutErrorMsg) const;
 
+	/** HTTP status code. */
+	uint32 StatusCode = 0;
 	/** HTTP response body. */
 	FSharedBuffer Body;
 };
diff --git a/Engine/Source/Developer/S3Client/Private/S3Client.cpp b/Engine/Source/Developer/S3Client/Private/S3Client.cpp
--- a/Engine/Source/Developer/S3Client/Private/S3Client.cpp
+++ b/Engine/Source/Developer/S3Client/Private/S3Client.cpp
@@ -34,8 +34,6 @@
 
 IMPLEMENT_MODULE(FDefaultModuleImpl, S3Client);
 
-DEFINE_LOG_CATEGORY_STATIC(LogS3Client, Log, All);
-
 namespace UE
 {
 
@@ -327,69 +325,7 @@ FS3CredentialsProfileStore FS3CredentialsProfileStore::FromFile(const FString& F
 }
 
 ////////////////////////////////////////////////////////////////////////////////
-
-FS3Response::FS3Response() = default;
-
-FS3Response::FS3Response(const FS3Response& Other) = default;
-
-FS3Response::FS3Response(FS3Response&& Other)
-	: HttpStatusCode(Other.HttpStatusCode)
-	, ApiStatusCode(Other.ApiStatusCode)
-	, Body(MoveTemp(Other.Body))
-{
-
-}
-
-FS3Response::FS3Response(uint32 InHttpStatusCode, uint32 InApiStatusCode)
-	: HttpStatusCode(InHttpStatusCode)
-	, ApiStatusCode(InApiStatusCode)
-{
-
-}
-
-FS3Response::FS3Response(uint32 InHttpStatusCode, uint32 InApiStatusCode, FSharedBuffer&& InBody)
-	: HttpStatusCode(InHttpStatusCode)
-	, ApiStatusCode(InApiStatusCode)
-	, Body(MoveTemp(InBody))
-{
-
-}
-
-FS3Response::FS3Response(uint32 InHttpStatusCode, FS3Response&& Other)
-	: HttpStatusCode(InHttpStatusCode)
-	, ApiStatusCode(Other.ApiStatusCode)
-	, Body(MoveTemp(Other.Body))
-{
-
-}
-
-FS3Response& FS3Response::operator=(const FS3Response& Other) = default;
-
-FS3Response& FS3Response::operator=(FS3Response&& Other)
-{
-	HttpStatusCode = Other.HttpStatusCode;
-	ApiStatusCode = Other.ApiStatusCode;
-	Body = MoveTemp(Other.Body);
-
-	return *this;
-}
-
-bool FS3Response::IsOk() const
-{
-	return ApiStatusCode == CURLE_OK && HttpStatusCode > 199 && HttpStatusCode < 299;
-}
-
-FSharedBuffer FS3Response::GetBody() const
-{
-	return Body;
-}
-
-FString FS3Response::ToString() const
-{
-	return FString(reinterpret_cast<const ANSICHAR*>(Body.GetData()));
-}
-
-void FS3Response::GetErrorResponse(FStringBuilderBase& OutErrorMsg) const
+void FS3Response::GetErrorMsg(FStringBuilderBase& OutErrorMsg) const
 {
 	OutErrorMsg.Reset();
 
@@ -399,66 +335,37 @@ void FS3Response::GetErrorResponse(FStringBuilderBase& OutErrorMsg) const
 		return;
 	}
 
-	if(ApiStatusCode != CURLE_OK)
-	{
-		OutErrorMsg << curl_easy_strerror(static_cast<CURLcode>(ApiStatusCode));
-	}
-	else
-	{
-		OutErrorMsg << TEXT("StatusCode: ") << HttpStatusCode << TEXT("Error: ");
-
-		const FString BodyResponseString = ToString();
-
-		FXmlFile XmlFile;
-		if (!XmlFile.LoadFile(BodyResponseString, EConstructMethod::ConstructFromBuffer))
-		{
-			OutErrorMsg << TEXT("Unknown");
-			return;
-		}
-
-		const FXmlNode* RootNode = XmlFile.GetRootNode();
-		if (!RootNode)
-		{
-			OutErrorMsg << TEXT("Unknown");
-			return;
-		}
-
-		const FXmlNode* CodeNode = RootNode->FindChildNode(TEXT("Code"));
-		if (!CodeNode)
-		{
-			OutErrorMsg << TEXT("Unknown");
-			return;
-		}
+	const FString BodyResponseString = ToString();
 
-		const FXmlNode* MessageNode = RootNode->FindChildNode(TEXT("Message"));
-		if (!MessageNode)
-		{
-			OutErrorMsg << TEXT("Unknown");
-			return;
-		}
-
-		OutErrorMsg << CodeNode->GetContent() << TEXT(": ") << MessageNode->GetContent();
+	FXmlFile XmlFile;
+	if (!XmlFile.LoadFile(BodyResponseString, EConstructMethod::ConstructFromBuffer))
+	{
+		OutErrorMsg << TEXT("Unknown");
+		return;
 	}
-}
-
-FString FS3Response:: GetErrorStatus() const
-{
-	TStringBuilder<256> ErrorStatus;
 
-	if (IsOk())
+	const FXmlNode* RootNode = XmlFile.GetRootNode();
+	if (!RootNode)
 	{
-		ErrorStatus << TEXT("Successs");
+		OutErrorMsg << TEXT("Unknown");
+		return;
 	}
-	else if (ApiStatusCode != CURLE_OK)
+
+	const FXmlNode* CodeNode = RootNode->FindChildNode(TEXT("Code"));
+	if (!CodeNode)
 	{
-		ErrorStatus << curl_easy_strerror(static_cast<CURLcode>(ApiStatusCode));
+		OutErrorMsg << TEXT("Unknown");
+		return;
 	}
-	else
+
+	const FXmlNode* MessageNode = RootNode->FindChildNode(TEXT("Message"));
+	if (!MessageNode)
 	{
-		ErrorStatus << TEXT("StatusCode: ") << HttpStatusCode << TEXT("Error: ");
+		OutErrorMsg << TEXT("Unknown");
+		return;
 	}
 
-	return FString(ErrorStatus);
+	OutErrorMsg << CodeNode->GetContent() << TEXT(": ") << MessageNode->GetContent();
 }
 
 ////////////////////////////////////////////////////////////////////////////////
@@ -676,10 +583,6 @@ FS3Response FS3Client::FS3Request::Perform(EMethod Method, const ANSICHAR* Url,
 	curl_easy_setopt(Curl, CURLOPT_WRITEDATA, this);
 	curl_easy_setopt(Curl, CURLOPT_WRITEFUNCTION, &WriteBodyCallback);
 
-	// Errors and logging
-	ANSICHAR ErrorBuffer[CURL_ERROR_SIZE] = {0};
-	curl_easy_setopt(Curl, CURLOPT_ERRORBUFFER, ErrorBuffer);
-
 	// SSL options
 	curl_easy_setopt(Curl, CURLOPT_USE_SSL, CURLUSESSL_ALL);
 	curl_easy_setopt(Curl, CURLOPT_SSL_VERIFYPEER, 1);
@@ -704,31 +607,17 @@ FS3Response FS3Client::FS3Request::Perform(EMethod Method, const ANSICHAR* Url,
 	{
 		CurlResult = curl_easy_getinfo(Curl, CURLINFO_RESPONSE_CODE, &ResponseCode);
 	}
-	
-	if (CurlResult != CURLE_OK)
-	{
-		UE_LOG(LogS3Client, Error, TEXT("%hs"), ErrorBuffer);
-	}
-
-	const char* ErrorMsg = curl_easy_strerror(CurlResult);
 
 	if (const uint64 Size = ResponseBody.TotalSize(); Size > 0)
 	{
 		return FS3Response
 		{
 			static_cast<uint32>(ResponseCode),
-			static_cast<uint32>(CurlResult),
 			FSharedBuffer::TakeOwnership(ResponseBody.ReleaseOwnership(), Size, FMemory::Free)
 		};
 	}
-	else
-	{
-		return FS3Response
-		{
-			static_cast<uint32>(ResponseCode),
-			static_cast<uint32>(CurlResult)
-		};
-	}
+
+	return FS3Response{static_cast<uint32>(ResponseCode)};
 }
 
 int FS3Client::FS3Request::StatusCallback(void* Ptr, curl_off_t TotalDownloadSize, curl_off_t CurrentDownloadSize, curl_off_t TotalUploadSize, curl_off_t CurrentUploadSize)
@@ -945,25 +834,25 @@ FS3ListObjectResponse FS3Client::ListObjects(const FS3ListObjectsRequest& ListRe
 
 	FS3Request Request(*this);
 	FS3Response Response = Request.Perform(FS3Request::EMethod::Get, Url.ToString(), FSharedBuffer());
-	if (!Response.IsOk())
+	if (Response.StatusCode != 200)
 	{
-		return FS3ListObjectResponse(MoveTemp(Response));
+		return FS3ListObjectResponse{{Response.StatusCode, MoveTemp(Response.Body)}};
 	}
 
-	FString Body(Response.ToString());
+	FString Body(reinterpret_cast<const ANSICHAR*>(Response.Body.GetData()));
 
 	FXmlFile XmlFile(Body, EConstructMethod::ConstructFromBuffer);
 	if (!XmlFile.IsValid())
 	{
 		//TODO: Better error message
-		return FS3ListObjectResponse{{500, MoveTemp(Response)}};
+		return FS3ListObjectResponse{{500, MoveTemp(Response.Body)}};
 	}
 
 	const FXmlNode* Root = XmlFile.GetRootNode();
 	if (!Root)
 	{
 		//TODO: Better error message
-		return FS3ListObjectResponse{{500, MoveTemp(Response)}};
+		return FS3ListObjectResponse{{500, MoveTemp(Response.Body)}};
 	}
 
 	FString BucketName;
@@ -999,7 +888,7 @@ FS3ListObjectResponse FS3Client::ListObjects(const FS3ListObjectsRequest& ListRe
 		}
 	}
 
-	return FS3ListObjectResponse{{200, 0, FSharedBuffer()}, MoveTemp(BucketName), MoveTemp(Objects)};
+	return FS3ListObjectResponse{{200, FSharedBuffer()}, MoveTemp(BucketName), MoveTemp(Objects)};
 }
 
 FS3DeleteObjectResponse FS3Client::DeleteObject(const FS3DeleteObjectRequest& DeleteRequest)
diff --git a/Engine/Source/Runtime/Experimental/IoStoreOnDemand/Private/IoStoreOnDemand.cpp b/Engine/Source/Runtime/Experimental/IoStoreOnDemand/Private/IoStoreOnDemand.cpp
--- a/Engine/Source/Runtime/Experimental/IoStoreOnDemand/Private/IoStoreOnDemand.cpp
+++ b/Engine/Source/Runtime/Experimental/IoStoreOnDemand/Private/IoStoreOnDemand.cpp
@@ -897,9 +897,9 @@ void FS3UploadQueue::ThreadEntry()
 		else
 		{
 			TStringBuilder<256> ErrorResponse;
-			Response.GetErrorResponse(ErrorResponse);
+			Response.GetErrorMsg(ErrorResponse);
 
-			UE_LOG(LogIas, Warning, TEXT("Failed to upload chunk '%s/%s/%s' (%s)"), *Client.GetConfig().ServiceUrl, *Bucket, *Entry.Key, ErrorResponse.ToString());
+			UE_LOG(LogIas, Warning, TEXT("Failed to upload chunk '%s/%s/%s', StatusCode: %u, Error: %s"), *Client.GetConfig().ServiceUrl, *Bucket, *Entry.Key, Response.StatusCode, ErrorResponse.ToString());
 			ErrorCount++;
 			break;
 		}
@@ -1268,13 +1268,13 @@ TIoStatusOr<FIoStoreUploadResult> UploadContainerFiles(
 			}
 
 			FOnDemandToc Toc;
-			FMemoryReaderView Ar(TocResponse.GetBody().GetView());
+			FMemoryReaderView Ar(TocResponse.Body.GetView());
 			Ar << Toc;
 
 			if (Ar.IsError()) 
 			{
 				Toc = FOnDemandToc{};
-				if (LoadFromCompactBinary(FCbFieldView(TocResponse.GetBody().GetData()), Toc) == false)
+				if (LoadFromCompactBinary(FCbFieldView(TocResponse.Body.GetData()), Toc) == false)
 				{
 					UE_LOG(LogIas, Warning, TEXT("Failed to load TOC '%s/%s/%s'"), *Client.GetConfig().ServiceUrl, *UploadParams.Bucket, *TocInfo.Key);
 					continue;
@@ -1456,7 +1456,7 @@ TIoStatusOr<FIoStoreUploadResult> UploadContainerFiles(
 
 			if (!Response.IsOk())
 			{
-				return FIoStatus(EIoErrorCode::WriteError, FString::Printf(TEXT("Failed to upload '%s' (%s)"), *UTocFilePath, *Response.GetErrorStatus()));
+				return FIoStatus(EIoErrorCode::WriteError, FString::Printf(TEXT("Failed to upload '%s', StatusCode: %u"), *UTocFilePath, Response.StatusCode));
 			}
 			
 			UTocPaths.Add(ContainerEntry.UTocHash, UTocFilePath);
@@ -1519,7 +1519,7 @@ TIoStatusOr<FIoStoreUploadResult> UploadContainerFiles(
 		}
 		else
 		{
-			UE_LOG(LogIas, Warning, TEXT("Failed to upload TOC '%s/%s/%s' (%s)"), *Client.GetConfig().ServiceUrl, *UploadParams.Bucket, Key.ToString(), *Response.GetErrorStatus());
+			UE_LOG(LogIas, Warning, TEXT("Failed to upload TOC '%s/%s/%s', StatusCode: %u"), *Client.GetConfig().ServiceUrl, *UploadParams.Bucket, Key.ToString(), Response.StatusCode);
 			return FIoStatus(EIoErrorCode::WriteError, TEXT("Failed to upload TOC"));
 		}
 
@@ -1716,7 +1716,7 @@ FIoStatus DownloadContainerFiles(const FIoStoreDownloadParams& DownloadParams, c
 	}
 
 	FOnDemandToc OnDemandToc;
-	if (LoadFromCompactBinary(FCbFieldView(TocResponse.GetBody().GetData()), OnDemandToc) == false)
+	if (LoadFromCompactBinary(FCbFieldView(TocResponse.Body.GetData()), OnDemandToc) == false)
 	{
 		return FIoStatus(EIoErrorCode::ReadError, TEXT("Failed to load on demand TOC"));
 	}
@@ -1755,12 +1755,12 @@ FIoStatus DownloadContainerFiles(const FIoStoreDownloadParams& DownloadParams, c
 		const FString UTocPath = FString::Printf(TEXT("%s/%s.utoc"), *DownloadParams.Directory, *ContainerEntry.ContainerName);
 		const FString UCasPath = FPaths::ChangeExtension(UTocPath, TEXT(".ucas")); 
 		FContainerStats& ContainerStats = ContainerSummary.FindOrAdd(ContainerEntry.ContainerName);
-		ContainerStats.TocRawSize = Response.GetBody().GetSize();
+		ContainerStats.TocRawSize = Response.Body.GetSize();
 
 		if (TUniquePtr<FArchive> TocFile(IFileManager::Get().CreateFileWriter(*UTocPath)); TocFile.IsValid())
 		{
 			UE_LOG(LogIas, Display, TEXT("Writing '%s'"), *UTocPath);
-			TocFile->Serialize((void*)Response.GetBody().GetData(), Response.GetBody().GetSize());
+			TocFile->Serialize((void*)Response.Body.GetData(), Response.Body.GetSize());
 		}
 		else
 		{
@@ -1850,10 +1850,10 @@ FIoStatus DownloadContainerFiles(const FIoStoreDownloadParams& DownloadParams, c
 			}
 
 			UE_LOG(LogIas, Display, TEXT("Serializing chunk %d/%d '%s' -> '%s' (%llu B)"),
-				Idx + 1, TocEntryCount, *HashString, *LexToString(ChunkId), ChunkResponse.GetBody().GetSize());
+				Idx + 1, TocEntryCount, *HashString, *LexToString(ChunkId), ChunkResponse.Body.GetSize());
 
 			check(CasFile->Tell() == FirstBlock.GetOffset());
-			CasFile->Serialize((void*)ChunkResponse.GetBody().GetData(), ChunkResponse.GetBody().GetSize());
+			CasFile->Serialize((void*)ChunkResponse.Body.GetData(), ChunkResponse.Body.GetSize());
 		}
 
 		ContainerStats.CompressedSize = CasFile->Tell();
@@ -2119,7 +2119,7 @@ FIoStatus ListTocs(const FIoStoreListTocsParams& Params)
 
 		if (Response.Objects.IsEmpty())
 		{
-			UE_LOG(LogIas, Display, TEXT("Not TOC's found at '%s/%s/%s' (%s)"), *Client.GetConfig().ServiceUrl, *Params.Bucket, Path.ToString(), *Response.GetErrorStatus());
+			UE_LOG(LogIas, Display, TEXT("Not TOC's found at '%s/%s/%s' (StatusCode: %u)"), *Client.GetConfig().ServiceUrl, *Params.Bucket, Path.ToString(), Response.StatusCode);
 			return FIoStatus(EIoErrorCode::NotFound);
 		}
 
@@ -2150,7 +2150,7 @@ FIoStatus ListTocs(const FIoStoreListTocsParams& Params)
 			}
 
 			FOnDemandToc Toc;
-			FMemoryReaderView Ar(TocResponse.GetBody().GetView());
+			FMemoryReaderView Ar(TocResponse.Body.GetView());
 			Ar << Toc;
 
 			if (Toc.Header.Magic != FOnDemandTocHeader::ExpectedMagic)
@@ -2196,7 +2196,7 @@ FIoStatus ListTocs(const FIoStoreListTocsParams& Params)
 		if (TocResponse.IsOk())
 		{
 			FOnDemandToc Toc;
-			FMemoryReaderView Ar(TocResponse.GetBody().GetView());
+			FMemoryReaderView Ar(TocResponse.Body.GetView());
 			Ar << Toc;
 
 			if (Toc.Header.Magic == FOnDemandTocHeader::ExpectedMagic)
@@ -2207,7 +2207,7 @@ FIoStatus ListTocs(const FIoStoreListTocsParams& Params)
 					.Toc = MoveTemp(Toc),
 					.DateTime = DateTime,
 					.Key = Params.TocKey,
-					.Size = TocResponse.GetBody().GetSize()
+					.Size = TocResponse.Body.GetSize()
 				});
 			}
 			else
