S3Client - don't append empty bucket name.

#rnx

diff --git a/Engine/Source/Developer/S3Client/Private/S3Client.cpp b/Engine/Source/Developer/S3Client/Private/S3Client.cpp
--- a/Engine/Source/Developer/S3Client/Private/S3Client.cpp
+++ b/Engine/Source/Developer/S3Client/Private/S3Client.cpp
@@ -42,6 +42,25 @@ namespace UE
 namespace S3Client
 {
 
+////////////////////////////////////////////////////////////////////////////////
+static FAnsiStringBuilderBase& AppendUrl(
+	FAnsiStringBuilderBase& Url,
+	const FString& Host,
+	const FString& BucketName,
+	const FString& Key)
+{
+	check(!Host.IsEmpty() && !Key.IsEmpty());
+
+	Url << StringCast<ANSICHAR>(*Host);
+	if (!BucketName.IsEmpty())
+	{
+		Url << "/" << StringCast<ANSICHAR>(*BucketName);
+	}
+	Url << "/" << StringCast<ANSICHAR>(*Key);
+
+	return Url;
+}
+
 ////////////////////////////////////////////////////////////////////////////////
 struct FSHA256
 {
@@ -889,7 +908,7 @@ void FS3Client::Teardown(FS3Request& Request)
 FS3GetObjectResponse FS3Client::GetObject(const FS3GetObjectRequest& GetRequest)
 {
 	TAnsiStringBuilder<256> Url;
-	Url << StringCast<ANSICHAR>(*Config.ServiceUrl) << "/" << StringCast<ANSICHAR>(*GetRequest.BucketName) << "/" << StringCast<ANSICHAR>(*GetRequest.Key);
+	S3Client::AppendUrl(Url, Config.ServiceUrl, GetRequest.BucketName, GetRequest.Key);
 
 	FS3Request Request(*this);
 	return Request.Perform(FS3Request::EMethod::Get, Url.ToString(), FSharedBuffer());
@@ -898,7 +917,7 @@ FS3GetObjectResponse FS3Client::GetObject(const FS3GetObjectRequest& GetRequest)
 FS3PutObjectResponse FS3Client::PutObject(const FS3PutObjectRequest& PutRequest)
 {
 	TAnsiStringBuilder<256> Url;
-	Url << StringCast<ANSICHAR>(*Config.ServiceUrl) << "/" << StringCast<ANSICHAR>(*PutRequest.BucketName) << "/" << StringCast<ANSICHAR>(*PutRequest.Key);
+	S3Client::AppendUrl(Url, Config.ServiceUrl, PutRequest.BucketName, PutRequest.Key);
 
 	FS3Request Request(*this);
 	return Request.Perform(FS3Request::EMethod::Put, Url.ToString(), FSharedBuffer::MakeView(PutRequest.ObjectData));
@@ -1005,7 +1024,7 @@ FS3ListObjectResponse FS3Client::ListObjects(const FS3ListObjectsRequest& ListRe
 FS3DeleteObjectResponse FS3Client::DeleteObject(const FS3DeleteObjectRequest& DeleteRequest)
 {
 	TAnsiStringBuilder<256> Url;
-	Url << StringCast<ANSICHAR>(*Config.ServiceUrl) << "/" << StringCast<ANSICHAR>(*DeleteRequest.BucketName) << "/" << StringCast<ANSICHAR>(*DeleteRequest.Key);
+	S3Client::AppendUrl(Url, Config.ServiceUrl, DeleteRequest.BucketName, DeleteRequest.Key);
 
 	FS3Request Request(*this);
 	return Request.Perform(FS3Request::EMethod::Delete, Url.ToString(), FSharedBuffer());
