Speculative crash fix when uploading IAS data.

#rnx

diff --git a/Engine/Source/Developer/S3Client/Private/S3Client.cpp b/Engine/Source/Developer/S3Client/Private/S3Client.cpp
--- a/Engine/Source/Developer/S3Client/Private/S3Client.cpp
+++ b/Engine/Source/Developer/S3Client/Private/S3Client.cpp
@@ -979,6 +979,11 @@ FS3ListObjectResponse FS3Client::ListObjects(const FS3ListObjectsRequest& ListRe
 		return FS3ListObjectResponse(MoveTemp(Response));
 	}
 
+	if (Response.GetBody().GetSize() == 0)
+	{
+		return FS3ListObjectResponse{MoveTemp(Response), FString(), TArray<FS3Object>()};
+	}
+
 	FString Body(Response.ToString());
 
 	FXmlFile XmlFile(Body, EConstructMethod::ConstructFromBuffer);
