S3 Client - don't rely on body being null terminated.

#rnx

diff --git a/Engine/Source/Developer/S3Client/Private/S3Client.cpp b/Engine/Source/Developer/S3Client/Private/S3Client.cpp
--- a/Engine/Source/Developer/S3Client/Private/S3Client.cpp
+++ b/Engine/Source/Developer/S3Client/Private/S3Client.cpp
@@ -405,7 +405,9 @@ FSharedBuffer FS3Response::GetBody() const
 
 FString FS3Response::ToString() const
 {
-	return FString(reinterpret_cast<const ANSICHAR*>(Body.GetData()));
+	return Body.GetSize() > 0
+		? FString::ConstructFromPtrSize(reinterpret_cast<const ANSICHAR*>(Body.GetData()), static_cast<int32>(Body.GetSize()))
+		: FString();
 }
 
 void FS3Response::GetErrorResponse(FStringBuilderBase& OutErrorMsg) const
