S3Client - add support for HeadObject.

#rnx

diff --git a/Engine/Source/Developer/S3Client/Internal/S3/S3Client.h b/Engine/Source/Developer/S3Client/Internal/S3/S3Client.h
--- a/Engine/Source/Developer/S3Client/Internal/S3/S3Client.h
+++ b/Engine/Source/Developer/S3Client/Internal/S3/S3Client.h
@@ -175,6 +175,18 @@ struct FS3GetObjectRequest
 /** Response parameters when retrieving objects. */
 using FS3GetObjectResponse = FS3Response;
 
+/** Request parameters for retrieving object meta data. */
+struct FS3HeadObjectRequest
+{
+	/** The bucket name. */
+	FString BucketName;
+	/** The object key. */
+	FString Key;
+};
+
+/** Response parameters when retrieving object meta data. */
+using FS3HeadObjectResponse = FS3Response;
+
 /** Request parameters for uploading objects. */
 struct FS3PutObjectRequest
 {
@@ -270,6 +282,8 @@ public:
 
 	/** Download an object associated with the specified request parameters. */
 	FS3GetObjectResponse GetObject(const FS3GetObjectRequest& Request);
+	/** Download object meta data. */
+	FS3HeadObjectResponse HeadObject(const FS3HeadObjectRequest& Request);
 	/** List all objects described by the specified request parameters. */
 	FS3ListObjectResponse ListObjects(const FS3ListObjectsRequest& Request);
 	/** Upload an object described by the specified request parameters. */
diff --git a/Engine/Source/Developer/S3Client/Private/S3Client.cpp b/Engine/Source/Developer/S3Client/Private/S3Client.cpp
--- a/Engine/Source/Developer/S3Client/Private/S3Client.cpp
+++ b/Engine/Source/Developer/S3Client/Private/S3Client.cpp
@@ -562,6 +562,7 @@ public:
 		Get,
 		Put,
 		Delete,
+		Head
 	};
 
 	FS3Request(FS3Client& S3Client);
@@ -605,7 +606,7 @@ FS3Response FS3Client::FS3Request::Perform(EMethod Method, const ANSICHAR* Url,
 
 	RequestBody = Body;
 	const uint64 ContentLength = Body.GetSize();
-	check((Method == EMethod::Get || Method == EMethod::Delete) || ContentLength > 0);
+	check((Method == EMethod::Get || Method == EMethod::Head || Method == EMethod::Delete) || ContentLength > 0);
 
 	// Find the host from the URL
 	const ANSICHAR* ProtocolEnd = FCStringAnsi::Strchr(Url, ':');
@@ -665,9 +666,13 @@ FS3Response FS3Client::FS3Request::Perform(EMethod Method, const ANSICHAR* Url,
 	// Setup the request
 	curl_easy_reset(Curl);
 	
-	if (Method == EMethod::Get)
+	if (Method == EMethod::Get || Method == EMethod::Head)
 	{
 		curl_easy_setopt(Curl, CURLOPT_HTTPGET, 1L);
+		if (Method == EMethod::Head)
+		{
+			curl_easy_setopt(Curl, CURLOPT_NOBODY, 1L);
+		}
 	}
 	else if (Method == EMethod::Put)
 	{
@@ -874,6 +879,8 @@ const ANSICHAR* FS3Client::FS3Request::LexToString(EMethod Method)
 		return "PUT";
 	case EMethod::Delete:
 		return "DELETE";
+	case EMethod::Head:
+		return "HEAD";
 	default:
 		check(false);
 		return nullptr;
@@ -916,6 +923,15 @@ FS3GetObjectResponse FS3Client::GetObject(const FS3GetObjectRequest& GetRequest)
 	return Request.Perform(FS3Request::EMethod::Get, Url.ToString(), FSharedBuffer());
 }
 
+FS3HeadObjectResponse FS3Client::HeadObject(const FS3HeadObjectRequest& GetRequest)
+{
+	TAnsiStringBuilder<256> Url;
+	S3Client::AppendUrl(Url, Config.ServiceUrl, GetRequest.BucketName, GetRequest.Key);
+
+	FS3Request Request(*this);
+	return Request.Perform(FS3Request::EMethod::Head, Url.ToString(), FSharedBuffer());
+}
+
 FS3PutObjectResponse FS3Client::PutObject(const FS3PutObjectRequest& PutRequest)
 {
 	TAnsiStringBuilder<256> Url;
