Add seeking callback so IAx uploads to S3 can be restarted if interrupted.

If the connection is interrupted for whatever reason, Curl retries. To do that it needs to rewind the put/post request payload. However this failed as there was no seek callback to facilitate that.

#rnx

diff --git a/Engine/Source/Developer/S3Client/Private/S3Client.cpp b/Engine/Source/Developer/S3Client/Private/S3Client.cpp
--- a/Engine/Source/Developer/S3Client/Private/S3Client.cpp
+++ b/Engine/Source/Developer/S3Client/Private/S3Client.cpp
@@ -575,6 +575,7 @@ private:
 	static size_t WriteHeadersCallback(void* Ptr, size_t SizeInBlocks, size_t BlockSizeInBytes, void* UserData);
 	static size_t ReadBodyCallback(void* Ptr, size_t SizeInBlocks, size_t BlockSizeInBytes, void* UserData);
 	static size_t WriteBodyCallback(void* Ptr, size_t SizeInBlocks, size_t BlockSizeInBytes, void* UserData);
+	static size_t SeekBodyCallback(void* Ptr, curl_off_t Offset, int Origin, void* UserData);
 	static int SslCertVerify(int PreverifyOk, X509_STORE_CTX* Context);
 	static CURLcode SslContextCallback(CURL* curl, void* sslctx, void* parm);
 	const ANSICHAR* LexToString(EMethod Method);
@@ -680,6 +681,8 @@ FS3Response FS3Client::FS3Request::Perform(EMethod Method, const ANSICHAR* Url,
 		curl_easy_setopt(Curl, CURLOPT_INFILESIZE, ContentLength);
 		curl_easy_setopt(Curl, CURLOPT_READDATA, this);
 		curl_easy_setopt(Curl, CURLOPT_READFUNCTION, &ReadBodyCallback);
+		curl_easy_setopt(Curl, CURLOPT_SEEKDATA, this);
+		curl_easy_setopt(Curl, CURLOPT_SEEKFUNCTION, &SeekBodyCallback);
 	}
 	else
 	{
@@ -791,6 +794,30 @@ size_t FS3Client::FS3Request::ReadBodyCallback(void* Ptr, size_t SizeInBlocks, s
 	return ReadSize;
 }
 
+size_t FS3Client::FS3Request::SeekBodyCallback(void* Ptr, curl_off_t Offset, int Origin, void* UserData)
+{
+	auto& Request = *static_cast<FS3Client::FS3Request*>(UserData);
+
+	uint64 Size = Request.RequestBody.GetSize();
+	uint64 Position = Request.BytesSent;
+
+	switch (Origin)
+	{
+	case SEEK_SET:	Position = Offset;				break;
+	case SEEK_CUR:	Position = Position + Offset;	break;
+	case SEEK_END:	Position = Size + Offset;		break;
+	default:		return CURL_SEEKFUNC_FAIL;
+	}
+
+	if (Position >= Size)
+	{
+		return CURL_SEEKFUNC_FAIL;
+	}
+
+	Request.BytesSent = Position;
+	return CURL_SEEKFUNC_OK;
+}
+
 size_t FS3Client::FS3Request::WriteBodyCallback(void* Ptr, size_t SizeInBlocks, size_t BlockSizeInBytes, void* UserData)
 {
 	const size_t WriteSize = SizeInBlocks * BlockSizeInBytes;
