Fix IAx upload crash

#rnx

diff --git a/Engine/Source/Developer/S3Client/Private/S3Client.cpp b/Engine/Source/Developer/S3Client/Private/S3Client.cpp
--- a/Engine/Source/Developer/S3Client/Private/S3Client.cpp
+++ b/Engine/Source/Developer/S3Client/Private/S3Client.cpp
@@ -575,7 +575,7 @@ private:
 	static size_t WriteHeadersCallback(void* Ptr, size_t SizeInBlocks, size_t BlockSizeInBytes, void* UserData);
 	static size_t ReadBodyCallback(void* Ptr, size_t SizeInBlocks, size_t BlockSizeInBytes, void* UserData);
 	static size_t WriteBodyCallback(void* Ptr, size_t SizeInBlocks, size_t BlockSizeInBytes, void* UserData);
-	static size_t SeekBodyCallback(void* Ptr, curl_off_t Offset, int Origin, void* UserData);
+	static size_t SeekBodyCallback(void* UserData, curl_off_t Offset, int Origin);
 	static int SslCertVerify(int PreverifyOk, X509_STORE_CTX* Context);
 	static CURLcode SslContextCallback(CURL* curl, void* sslctx, void* parm);
 	const ANSICHAR* LexToString(EMethod Method);
@@ -794,7 +794,7 @@ size_t FS3Client::FS3Request::ReadBodyCallback(void* Ptr, size_t SizeInBlocks, s
 	return ReadSize;
 }
 
-size_t FS3Client::FS3Request::SeekBodyCallback(void* Ptr, curl_off_t Offset, int Origin, void* UserData)
+size_t FS3Client::FS3Request::SeekBodyCallback(void* UserData, curl_off_t Offset, int Origin)
 {
 	auto& Request = *static_cast<FS3Client::FS3Request*>(UserData);
 
