Off-by-one error

#rnx

diff --git a/Engine/Source/Developer/S3Client/Private/S3Client.cpp b/Engine/Source/Developer/S3Client/Private/S3Client.cpp
--- a/Engine/Source/Developer/S3Client/Private/S3Client.cpp
+++ b/Engine/Source/Developer/S3Client/Private/S3Client.cpp
@@ -809,7 +809,7 @@ size_t FS3Client::FS3Request::SeekBodyCallback(void* UserData, curl_off_t Offset
 	default:		return CURL_SEEKFUNC_FAIL;
 	}
 
-	if (Position >= Size)
+	if (Position > Size)
 	{
 		return CURL_SEEKFUNC_FAIL;
 	}
