Add a new define UE_ENABLE_S3_DEBUGGING that can enable curl debugging for the S3 client. Disabled by default.

#rnx

diff --git a/Engine/Source/Developer/S3Client/Private/S3Client.cpp b/Engine/Source/Developer/S3Client/Private/S3Client.cpp
--- a/Engine/Source/Developer/S3Client/Private/S3Client.cpp
+++ b/Engine/Source/Developer/S3Client/Private/S3Client.cpp
@@ -36,12 +36,42 @@ IMPLEMENT_MODULE(FDefaultModuleImpl, S3Client);
 
 DEFINE_LOG_CATEGORY_STATIC(LogS3Client, Log, All);
 
+// Enable to start logging debug output from curl
+#define UE_ENABLE_S3_DEBUGGING 0
+
 namespace UE
 {
 
 namespace S3Client
 {
 
+#if UE_ENABLE_S3_DEBUGGING
+
+int DebugCallback(CURL* Handle, curl_infotype Type, char* Data, size_t Size, void* UserData)
+{
+	FString Message(Size, Data);
+	Message.TrimEndInline();
+
+	switch (Type)
+	{
+		case CURLINFO_TEXT:
+			UE_LOG(LogS3Client, Display, TEXT("Info: %s"), *Message);
+			break;
+		case CURLINFO_HEADER_OUT:
+			UE_LOG(LogS3Client, Display, TEXT("Send Header: %s"), *Message);
+			break;
+		case CURLINFO_HEADER_IN:
+			UE_LOG(LogS3Client, Display, TEXT("Recv Header: %s"), *Message);
+			break;
+		default:
+			break;
+	}
+
+	return 0;
+}
+
+#endif // UE_ENABLE_S3_DEBUGGING
+
 ////////////////////////////////////////////////////////////////////////////////
 static FAnsiStringBuilderBase& AppendUrl(
 	FAnsiStringBuilderBase& Url,
@@ -666,6 +696,12 @@ FS3Response FS3Client::FS3Request::Perform(EMethod Method, const ANSICHAR* Url,
 
 	// Setup the request
 	curl_easy_reset(Curl);
+
+#if UE_ENABLE_S3_DEBUGGING
+	curl_easy_setopt(Curl, CURLOPT_DEBUGFUNCTION, S3Client::DebugCallback);
+	curl_easy_setopt(Curl, CURLOPT_DEBUGDATA, nullptr);
+	curl_easy_setopt(Curl, CURLOPT_VERBOSE, 1L);
+#endif //UE_ENABLE_S3_DEBUGGING
 	
 	if (Method == EMethod::Get || Method == EMethod::Head)
 	{
@@ -1107,4 +1143,6 @@ FS3DeleteObjectResponse FS3Client::DeleteObject(const FS3DeleteObjectRequest& De
 
 } // namespace UE
 
+#undef UE_ENABLE_S3_DEBUGGING
+
 #endif // (IS_PROGRAM || WITH_EDITOR)
